Editor: Start extracting SSettingsEditorCheckoutNotice out as a standalone widget (WIP)

diff --git a/Engine/Source/Developer/SettingsEditor/Private/SettingsEditorPrivatePCH.h b/Engine/Source/Developer/SettingsEditor/Private/SettingsEditorPrivatePCH.h
--- a/Engine/Source/Developer/SettingsEditor/Private/SettingsEditorPrivatePCH.h
+++ b/Engine/Source/Developer/SettingsEditor/Private/SettingsEditorPrivatePCH.h
@@ -20,5 +20,4 @@
 
 #include "SettingsEditorModel.h"
 
-#include "SSettingsEditorCheckoutNotice.h"
 #include "SSettingsEditor.h"
diff --git a/Engine/Source/Developer/SettingsEditor/Private/Widgets/SSettingsEditor.cpp b/Engine/Source/Developer/SettingsEditor/Private/Widgets/SSettingsEditor.cpp
--- a/Engine/Source/Developer/SettingsEditor/Private/Widgets/SSettingsEditor.cpp
+++ b/Engine/Source/Developer/SettingsEditor/Private/Widgets/SSettingsEditor.cpp
@@ -4,6 +4,7 @@
 #include "EngineAnalytics.h"
 #include "AnalyticsEventAttribute.h"
 #include "IAnalyticsProvider.h"
+#include "SSettingsEditorCheckoutNotice.h"
 
 #define LOCTEXT_NAMESPACE "SSettingsEditor"
 
@@ -195,7 +196,6 @@ void SSettingsEditor::Construct( const FArguments& InArgs, const ISettingsEditor
 							[
 								// checkout notice
 								SNew(SSettingsEditorCheckoutNotice)
-									.OnCheckOutClicked(this, &SSettingsEditor::HandleCheckOutButtonClicked)
 									.Visibility(this, &SSettingsEditor::HandleDefaultConfigNoticeVisibility)
 									.Unlocked(this, &SSettingsEditor::HandleConfigNoticeUnlocked)
 									.ConfigFilePath(this, &SSettingsEditor::GetCachedConfigFileName)
diff --git a/Engine/Source/Developer/SettingsEditor/Private/Widgets/SSettingsEditorCheckoutNotice.cpp b/Engine/Source/Developer/SettingsEditor/Private/Widgets/SSettingsEditorCheckoutNotice.cpp
deleted file mode 100644
diff --git a/Engine/Source/Developer/SettingsEditor/Private/Widgets/SSettingsEditorCheckoutNotice.h b/Engine/Source/Developer/SettingsEditor/Private/Widgets/SSettingsEditorCheckoutNotice.h
deleted file mode 100644
diff --git a/Engine/Source/Developer/SettingsEditor/SettingsEditor.Build.cs b/Engine/Source/Developer/SettingsEditor/SettingsEditor.Build.cs
--- a/Engine/Source/Developer/SettingsEditor/SettingsEditor.Build.cs
+++ b/Engine/Source/Developer/SettingsEditor/SettingsEditor.Build.cs
@@ -29,6 +29,7 @@ public class SettingsEditor : ModuleRules
 				"SlateCore",
 				"SourceControl",
 				"Analytics",
+				"SharedSettingsWidgets",
 				"Engine",
             }
         );
diff --git a/Engine/Source/Developer/SharedSettingsWidgets/Private/SSettingsEditorCheckoutNotice.cpp b/Engine/Source/Developer/SharedSettingsWidgets/Private/SSettingsEditorCheckoutNotice.cpp
new file mode 100644
--- /dev/null
+++ b/Engine/Source/Developer/SharedSettingsWidgets/Private/SSettingsEditorCheckoutNotice.cpp
@@ -0,0 +1,218 @@
+// Copyright 1998-2014 Epic Games, Inc. All Rights Reserved.
+
+#include "SharedSettingsWidgetsPrivatePCH.h"
+#include "EditorStyle.h"
+#include "SSettingsEditorCheckoutNotice.h"
+#include "ISourceControlModule.h"
+
+
+#define LOCTEXT_NAMESPACE "SSettingsEditorCheckoutNotice"
+
+
+/* SSettingsEditorCheckoutNotice interface
+ *****************************************************************************/
+
+void SSettingsEditorCheckoutNotice::Construct( const FArguments& InArgs )
+{
+	CheckOutClickedDelegate = InArgs._OnCheckOutClicked;
+	bIsUnlocked = InArgs._Unlocked;
+	ConfigFilePath = InArgs._ConfigFilePath;
+	bLookingForSourceControlState = InArgs._LookingForSourceControlState;
+
+	// default configuration notice
+	ChildSlot
+	[
+		SNew(SBorder)
+			.BorderBackgroundColor(this, &SSettingsEditorCheckoutNotice::GetLockedOrUnlockedStatusBarColor)
+			.BorderImage(FEditorStyle::GetBrush("ToolPanel.LightGroupBorder"))
+			.Padding(8.0f)
+			[
+				SNew(SWidgetSwitcher)
+					.WidgetIndex(this, &SSettingsEditorCheckoutNotice::HandleNoticeSwitcherWidgetIndex)
+
+				// Locked slot
+				+ SWidgetSwitcher::Slot()
+					[
+						SNew(SHorizontalBox)
+
+						// Locked icon
+						+ SHorizontalBox::Slot()
+							.AutoWidth()
+							.VAlign(VAlign_Center)
+							[
+								SNew(SImage)
+									.Image(FEditorStyle::GetBrush("GenericLock"))
+							]
+
+						// Locked notice
+						+ SHorizontalBox::Slot()
+							.FillWidth(1.0f)
+							.Padding(16.0f, 0.0f)
+							.VAlign(VAlign_Center)
+							[
+								SNew(STextBlock)
+									.Text(this, &SSettingsEditorCheckoutNotice::HandleLockedStatusText)
+									.ColorAndOpacity(FLinearColor::White)
+									.ShadowColorAndOpacity(FLinearColor::Black)
+									.ShadowOffset(FVector2D::UnitVector)
+							]
+
+						// Check out button
+						+ SHorizontalBox::Slot()
+							.AutoWidth()
+							[
+								SNew(SButton)
+									.OnClicked(this, &SSettingsEditorCheckoutNotice::HandleCheckOutButtonClicked)
+									.Text(this, &SSettingsEditorCheckoutNotice::HandleCheckOutButtonText)
+									.ToolTipText(this, &SSettingsEditorCheckoutNotice::HandleCheckOutButtonToolTip)
+									.Visibility(this, &SSettingsEditorCheckoutNotice::HandleCheckOutButtonVisibility)
+							]
+
+						// Source control status throbber
+						+ SHorizontalBox::Slot()
+							.AutoWidth()
+							.Padding(0.0f, 1.0f)
+							[
+								SNew(SThrobber)
+									.Visibility(this, &SSettingsEditorCheckoutNotice::HandleThrobberVisibility)
+							]
+					]
+
+				// Unlocked slot
+				+ SWidgetSwitcher::Slot()
+					[
+						SNew(SHorizontalBox)
+
+						// Unlocked icon
+						+ SHorizontalBox::Slot()
+							.AutoWidth()
+							.VAlign(VAlign_Center)
+							[
+								SNew(SImage)
+									.Image(FEditorStyle::GetBrush("GenericUnlock"))
+							]
+
+						// Unlocked notice
+						+ SHorizontalBox::Slot()
+							.AutoWidth()
+							.Padding(16.0f, 0.0f)
+							.VAlign(VAlign_Center)
+							[
+								SNew(STextBlock)
+									.Text(this, &SSettingsEditorCheckoutNotice::HandleUnlockedStatusText)
+									.ColorAndOpacity(FLinearColor::White)
+									.ShadowColorAndOpacity(FLinearColor::Black)
+									.ShadowOffset(FVector2D::UnitVector)
+							]
+					]
+			]
+	];
+}
+
+
+/* SSettingsEditorCheckoutNotice callbacks
+ *****************************************************************************/
+
+FReply SSettingsEditorCheckoutNotice::HandleCheckOutButtonClicked()
+{
+	FString TargetFilePath = ConfigFilePath.Get();
+
+	if (ISourceControlModule::Get().IsEnabled())
+	{
+		FText ErrorMessage;
+
+		if (!SourceControlHelpers::CheckoutOrMarkForAdd(TargetFilePath, FText::FromString(TargetFilePath), NULL, ErrorMessage))
+		{
+			FNotificationInfo Info(ErrorMessage);
+			Info.ExpireDuration = 3.0f;
+			FSlateNotificationManager::Get().AddNotification(Info);
+		}
+	}
+	else
+	{
+		if (!FPlatformFileManager::Get().GetPlatformFile().SetReadOnly(*TargetFilePath, false))
+		{
+			FText NotificationErrorText = FText::Format(LOCTEXT("FailedToMakeWritable", "Could not make {0} writable."), FText::FromString(TargetFilePath));
+
+			FNotificationInfo Info(NotificationErrorText);
+			Info.ExpireDuration = 3.0f;
+
+			FSlateNotificationManager::Get().AddNotification(Info);
+		}
+	}
+
+	if (CheckOutClickedDelegate.IsBound())
+	{
+		return CheckOutClickedDelegate.Execute();
+	}
+
+	return FReply::Handled();
+}
+
+
+FText SSettingsEditorCheckoutNotice::HandleCheckOutButtonText( ) const
+{
+	if (ISourceControlModule::Get().IsEnabled() && ISourceControlModule::Get().GetProvider().IsAvailable())
+	{
+		return LOCTEXT("CheckOutFile", "Check Out File");
+	}
+
+	return LOCTEXT("MakeWritable", "Make Writable");
+}
+
+
+FText SSettingsEditorCheckoutNotice::HandleCheckOutButtonToolTip( ) const
+{
+	if (ISourceControlModule::Get().IsEnabled() && ISourceControlModule::Get().GetProvider().IsAvailable())
+	{
+		return LOCTEXT("CheckOutFileTooltip", "Check out the default configuration file that holds these settings.");
+	}
+
+	return LOCTEXT("MakeWritableTooltip", "Make the default configuration file that holds these settings writable.");
+}
+
+
+EVisibility SSettingsEditorCheckoutNotice::HandleCheckOutButtonVisibility( ) const
+{
+	if (ISourceControlModule::Get().IsEnabled() && ISourceControlModule::Get().GetProvider().IsAvailable())
+	{
+		return !bLookingForSourceControlState.Get() ? EVisibility::Visible : EVisibility::Collapsed;
+	}
+
+	return EVisibility::Collapsed;
+}
+
+FText SSettingsEditorCheckoutNotice::HandleLockedStatusText() const
+{
+	FText ConfigFilename = FText::FromString(FPaths::GetCleanFilename(ConfigFilePath.Get()));
+
+	return FText::Format(ISourceControlModule::Get().IsEnabled() ?
+		LOCTEXT("DefaultSettingsNotice_WithSourceControl", "These settings are saved in {0}, which is not currently checked out.") :
+		LOCTEXT("DefaultSettingsNotice_Source", "These settings are saved {0}, which is not currently writable."), ConfigFilename);
+}
+
+FText SSettingsEditorCheckoutNotice::HandleUnlockedStatusText() const
+{
+	FText ConfigFilename = FText::FromString(FPaths::GetCleanFilename(ConfigFilePath.Get()));
+	
+	return FText::Format(ISourceControlModule::Get().IsEnabled() ?
+		LOCTEXT("DefaultSettingsNotice_CheckedOut", "These settings are saved in {0}, which is currently checked out.") :
+		LOCTEXT("DefaultSettingsNotice_Writable", "These settings are saved in {0}, which is currently writable."), ConfigFilename);
+}
+
+FSlateColor SSettingsEditorCheckoutNotice::GetLockedOrUnlockedStatusBarColor() const
+{
+	return bIsUnlocked.Get() ? FLinearColor::Green : FLinearColor::Yellow;
+}
+
+
+EVisibility SSettingsEditorCheckoutNotice::HandleThrobberVisibility() const
+{
+	if(ISourceControlModule::Get().IsEnabled())
+	{
+		return bLookingForSourceControlState.Get() ? EVisibility::Visible : EVisibility::Collapsed;
+	}
+	return EVisibility::Collapsed;
+}
+
+#undef LOCTEXT_NAMESPACE
diff --git a/Engine/Source/Developer/SharedSettingsWidgets/Public/SSettingsEditorCheckoutNotice.h b/Engine/Source/Developer/SharedSettingsWidgets/Public/SSettingsEditorCheckoutNotice.h
new file mode 100644
--- /dev/null
+++ b/Engine/Source/Developer/SharedSettingsWidgets/Public/SSettingsEditorCheckoutNotice.h
@@ -0,0 +1,81 @@
+// Copyright 1998-2014 Epic Games, Inc. All Rights Reserved.
+
+#pragma once
+
+
+/**
+ * Implements a widget that provides a notice for files that need to be checked out.
+ */
+class SHAREDSETTINGSWIDGETS_API SSettingsEditorCheckoutNotice : public SCompoundWidget
+{
+public:
+
+	SLATE_BEGIN_ARGS(SSettingsEditorCheckoutNotice) { }
+
+		/** Called to determine if the associated file is unlocked */
+		SLATE_ATTRIBUTE(bool, Unlocked)
+
+		/** Called to get the filename of the config file for display */
+		SLATE_ATTRIBUTE(FString, ConfigFilePath)
+
+		/** Slot for this button's content (optional) */
+		SLATE_NAMED_SLOT(FArguments, LockedContent)
+
+		/** Called when the 'Check Out' button is clicked */
+		SLATE_EVENT(FOnClicked, OnCheckOutClicked)
+
+		/** Called to determine if we are currently looking for the source control state of the file */
+		SLATE_ATTRIBUTE(bool, LookingForSourceControlState)
+
+	SLATE_END_ARGS()
+
+	/**
+	 * Constructs the widget.
+	 *
+	 * @param InArgs The construction arguments.
+	 */
+	void Construct( const FArguments& InArgs );
+
+private:
+
+	// Callback for clicking the 'Check Out' button.
+	FReply HandleCheckOutButtonClicked( );
+
+	// Callback for getting the text of the 'Check Out' button.
+	FText HandleCheckOutButtonText( ) const;
+
+	// Callback for getting the tool tip text of the 'Check Out' button.
+	FText HandleCheckOutButtonToolTip( ) const;
+
+	// Callback for determining the visibility of the check-out button.
+	EVisibility HandleCheckOutButtonVisibility( ) const;
+
+	// Callback for getting the widget index for the notice switcher.
+	int32 HandleNoticeSwitcherWidgetIndex( ) const
+	{
+		return bIsUnlocked.Get() ? 1 : 0;
+	}
+
+	// Callback for getting the status text when the config is locked
+	FText HandleLockedStatusText() const;
+
+	// Callback for getting the status text when the config is unlocked
+	FText HandleUnlockedStatusText() const;
+
+	// Callback for getting the status color, which indicates if it is writeable or not
+	FSlateColor GetLockedOrUnlockedStatusBarColor() const;
+
+	// Callback for getting the visibility of the source control throbber
+	EVisibility HandleThrobberVisibility() const;
+
+private:
+
+	// Holds a delegate that is executed when the 'Check Out' button has been clicked.
+	FOnClicked CheckOutClickedDelegate;
+
+	TAttribute<bool> bIsUnlocked;
+
+	TAttribute<FString> ConfigFilePath;
+
+	TAttribute<bool> bLookingForSourceControlState;
+};
