Editor: Various settings editor fixes
- Fix SSettingsEditor reload config object regression
- Remove double-padding when editing a settings page that has a checkout prompt
- Fix pending hardware targeting settings not showing up if you change pages and come back
- Fix height pop in SSetingsEditorCheckoutNotice when status query finishes

diff --git a/Engine/Source/Developer/SettingsEditor/Private/Widgets/SSettingsEditor.cpp b/Engine/Source/Developer/SettingsEditor/Private/Widgets/SSettingsEditor.cpp
--- a/Engine/Source/Developer/SettingsEditor/Private/Widgets/SSettingsEditor.cpp
+++ b/Engine/Source/Developer/SettingsEditor/Private/Widgets/SSettingsEditor.cpp
@@ -50,6 +50,7 @@ void SSettingsEditor::Construct( const FArguments& InArgs, const ISettingsEditor
 	FileWatcherWidget =
 		SNew(SSettingsEditorCheckoutNotice)
 		.Visibility(this, &SSettingsEditor::HandleDefaultConfigNoticeVisibility)
+		.OnFileProbablyModifiedExternally(this, &SSettingsEditor::ReloadConfigObject)
 		.ConfigFilePath(this, &SSettingsEditor::GetConfigFileName);
 
 	ChildSlot
@@ -195,7 +196,7 @@ void SSettingsEditor::Construct( const FArguments& InArgs, const ISettingsEditor
 
 						+ SVerticalBox::Slot()
 							.AutoHeight()
-							.Padding(0.0f, 16.0f, 0.0f, 16.0f)
+							.Padding(0.0f, 16.0f, 0.0f, 0.0f)
 							[
 								FileWatcherWidget.ToSharedRef()
 							]
@@ -830,4 +831,14 @@ bool SSettingsEditor::IsDefaultConfigCheckOutNeeded() const
 	}
 }
 
+void SSettingsEditor::ReloadConfigObject()
+{
+	TWeakObjectPtr<UObject> SettingsObject = GetSelectedSettingsObject();
+
+	if (SettingsObject.IsValid() && SettingsObject->GetClass()->HasAnyClassFlags(CLASS_Config | CLASS_DefaultConfig))
+	{
+		SettingsObject->ReloadConfig();
+	}
+}
+
 #undef LOCTEXT_NAMESPACE
diff --git a/Engine/Source/Developer/SettingsEditor/Private/Widgets/SSettingsEditor.h b/Engine/Source/Developer/SettingsEditor/Private/Widgets/SSettingsEditor.h
--- a/Engine/Source/Developer/SettingsEditor/Private/Widgets/SSettingsEditor.h
+++ b/Engine/Source/Developer/SettingsEditor/Private/Widgets/SSettingsEditor.h
@@ -163,6 +163,9 @@ private:
 
 	// Do we need to edit the default config file?
 	bool IsDefaultConfigCheckOutNeeded() const;
+
+	void ReloadConfigObject();
+
 private:
 
 	// Watcher widget for the default config file (checks file status / SCC state)
diff --git a/Engine/Source/Developer/SharedSettingsWidgets/Private/SSettingsEditorCheckoutNotice.cpp b/Engine/Source/Developer/SharedSettingsWidgets/Private/SSettingsEditorCheckoutNotice.cpp
--- a/Engine/Source/Developer/SharedSettingsWidgets/Private/SSettingsEditorCheckoutNotice.cpp
+++ b/Engine/Source/Developer/SharedSettingsWidgets/Private/SSettingsEditorCheckoutNotice.cpp
@@ -14,7 +14,7 @@
 
 void SSettingsEditorCheckoutNotice::Construct( const FArguments& InArgs )
 {
-	CheckOutClickedDelegate = InArgs._OnCheckOutClicked;
+	OnFileProbablyModifiedExternally = InArgs._OnFileProbablyModifiedExternally;
 	ConfigFilePath = InArgs._ConfigFilePath;
 
 	DefaultConfigCheckOutTimer = 0.0f;
@@ -143,11 +143,6 @@ FReply SSettingsEditorCheckoutNotice::HandleCheckOutButtonClicked()
 		}
 	}
 
-	if (CheckOutClickedDelegate.IsBound())
-	{
-		return CheckOutClickedDelegate.Execute();
-	}
-
 	return FReply::Handled();
 }
 
@@ -180,7 +175,7 @@ EVisibility SSettingsEditorCheckoutNotice::HandleCheckOutButtonVisibility( ) con
 		return !DefaultConfigQueryInProgress ? EVisibility::Visible : EVisibility::Hidden;
 	}
 
-	return EVisibility::Collapsed;
+	return EVisibility::Hidden;
 }
 
 FText SSettingsEditorCheckoutNotice::HandleLockedStatusText() const
@@ -250,13 +245,7 @@ void SSettingsEditorCheckoutNotice::Tick(const FGeometry& AllottedGeometry, cons
 			// file has been checked in or reverted
 			if ((NewCheckOutNeeded == true) && (DefaultConfigCheckOutNeeded == false))
 			{
-				// 				TWeakObjectPtr<UObject> SettingsObject = GetSelectedSettingsObject();
-				// 
-				// 				if (SettingsObject.IsValid() && SettingsObject->GetClass()->HasAnyClassFlags(CLASS_Config | CLASS_DefaultConfig))
-				// 				{
-				// 					//@TODO: This thing!
-				// 					SettingsObject->ReloadConfig();
-				// 				}
+				OnFileProbablyModifiedExternally.ExecuteIfBound();
 			}
 		}
 
diff --git a/Engine/Source/Developer/SharedSettingsWidgets/Public/SSettingsEditorCheckoutNotice.h b/Engine/Source/Developer/SharedSettingsWidgets/Public/SSettingsEditorCheckoutNotice.h
--- a/Engine/Source/Developer/SharedSettingsWidgets/Public/SSettingsEditorCheckoutNotice.h
+++ b/Engine/Source/Developer/SharedSettingsWidgets/Public/SSettingsEditorCheckoutNotice.h
@@ -12,20 +12,11 @@ public:
 
 	SLATE_BEGIN_ARGS(SSettingsEditorCheckoutNotice) { }
 
-		/** Called to determine if the associated file is unlocked */
-		SLATE_ATTRIBUTE(bool, Unlocked)
-
 		/** Called to get the filename of the config file for display */
 		SLATE_ATTRIBUTE(FString, ConfigFilePath)
 
-		/** Slot for this button's content (optional) */
-		SLATE_NAMED_SLOT(FArguments, LockedContent)
-
-		/** Called when the 'Check Out' button is clicked */
-		SLATE_EVENT(FOnClicked, OnCheckOutClicked)
-
-		/** Called to determine if we are currently looking for the source control state of the file */
-		SLATE_ATTRIBUTE(bool, LookingForSourceControlState)
+		/** Called when the the file was probably modified externally (only detects files that were 'ready to go' and are now not, e.g., reverted/checked in or files with the read-only flag edited outside) */
+		SLATE_EVENT(FSimpleDelegate, OnFileProbablyModifiedExternally)
 
 	SLATE_END_ARGS()
 
@@ -83,8 +74,8 @@ private:
 
 private:
 
-	// Holds a delegate that is executed when the 'Check Out' button has been clicked.
-	FOnClicked CheckOutClickedDelegate;
+	// Called when the file was modified externally
+	FSimpleDelegate OnFileProbablyModifiedExternally;
 
 	// The current file being watched
 	TAttribute<FString> ConfigFilePath;
diff --git a/Engine/Source/Editor/DetailCustomizations/Private/HardwareTargetingSettingsDetails.cpp b/Engine/Source/Editor/DetailCustomizations/Private/HardwareTargetingSettingsDetails.cpp
--- a/Engine/Source/Editor/DetailCustomizations/Private/HardwareTargetingSettingsDetails.cpp
+++ b/Engine/Source/Editor/DetailCustomizations/Private/HardwareTargetingSettingsDetails.cpp
@@ -120,6 +120,8 @@ public:
 					EditPropertiesBlock
 				];
 		}
+
+		Update();
 	}
 
 	bool CanApply() const
