Tidied up ini file checkout warnings.
The colored background has been removed in favor of a less intrusive coloration on the icon itself. Updated the padlock icons.

diff --git a/Engine/Source/Developer/SharedSettingsWidgets/Private/SSettingsEditorCheckoutNotice.cpp b/Engine/Source/Developer/SharedSettingsWidgets/Private/SSettingsEditorCheckoutNotice.cpp
--- a/Engine/Source/Developer/SharedSettingsWidgets/Private/SSettingsEditorCheckoutNotice.cpp
+++ b/Engine/Source/Developer/SharedSettingsWidgets/Private/SSettingsEditorCheckoutNotice.cpp
@@ -21,13 +21,12 @@ void SSettingsEditorCheckoutNotice::Construct( const FArguments& InArgs )
 	DefaultConfigCheckOutNeeded = false;
 	DefaultConfigQueryInProgress = true;
 
+	const int Padding = 8;
 	// default configuration notice
 	ChildSlot
 	[
 		SNew(SBorder)
-			.BorderBackgroundColor(this, &SSettingsEditorCheckoutNotice::GetLockedOrUnlockedStatusBarColor)
-			.BorderImage(FEditorStyle::GetBrush("ToolPanel.LightGroupBorder"))
-			.Padding(8.0f)
+			.BorderImage(FEditorStyle::GetBrush("ToolPanel.GroupBorder"))
 			[
 				SNew(SWidgetSwitcher)
 					.WidgetIndex(this, &SSettingsEditorCheckoutNotice::HandleNoticeSwitcherWidgetIndex)
@@ -40,28 +39,30 @@ void SSettingsEditorCheckoutNotice::Construct( const FArguments& InArgs )
 						// Locked icon
 						+ SHorizontalBox::Slot()
 							.AutoWidth()
+							.Padding(Padding)
 							.VAlign(VAlign_Center)
 							[
 								SNew(SImage)
-									.Image(FEditorStyle::GetBrush("GenericLock"))
+								.Image(FEditorStyle::GetBrush("GenericLock"))
+								.ColorAndOpacity(FColor(255, 168, 0))
 							]
 
 						// Locked notice
 						+ SHorizontalBox::Slot()
 							.FillWidth(1.0f)
-							.Padding(16.0f, 0.0f)
+							.Padding(FMargin(0.f, Padding, Padding, Padding))
 							.VAlign(VAlign_Center)
 							[
 								SNew(STextBlock)
 									.Text(this, &SSettingsEditorCheckoutNotice::HandleLockedStatusText)
-									.ColorAndOpacity(FLinearColor::White)
-									.ShadowColorAndOpacity(FLinearColor::Black)
-									.ShadowOffset(FVector2D::UnitVector)
+									.Font(FEditorStyle::GetFontStyle("PropertyWindow.NormalFont"))
 							]
 
 						// Check out button
 						+ SHorizontalBox::Slot()
+							.Padding(FMargin(0.f, 0.f, Padding, 0.f))
 							.AutoWidth()
+							.VAlign(VAlign_Center)
 							[
 								SNew(SButton)
 									.OnClicked(this, &SSettingsEditorCheckoutNotice::HandleCheckOutButtonClicked)
@@ -73,6 +74,7 @@ void SSettingsEditorCheckoutNotice::Construct( const FArguments& InArgs )
 						// Source control status throbber
 						+ SHorizontalBox::Slot()
 							.AutoWidth()
+							.VAlign(VAlign_Center)
 							.Padding(0.0f, 1.0f)
 							[
 								SNew(SThrobber)
@@ -88,23 +90,23 @@ void SSettingsEditorCheckoutNotice::Construct( const FArguments& InArgs )
 						// Unlocked icon
 						+ SHorizontalBox::Slot()
 							.AutoWidth()
+							.Padding(Padding)
 							.VAlign(VAlign_Center)
 							[
 								SNew(SImage)
-									.Image(FEditorStyle::GetBrush("GenericUnlock"))
+								.Image(FEditorStyle::GetBrush("GenericUnlock"))
+								.ColorAndOpacity(FLinearColor::Green)
 							]
 
 						// Unlocked notice
 						+ SHorizontalBox::Slot()
 							.AutoWidth()
-							.Padding(16.0f, 0.0f)
+							.Padding(FMargin(0.f, Padding, Padding, Padding))
 							.VAlign(VAlign_Center)
 							[
 								SNew(STextBlock)
+									.Font(FEditorStyle::GetFontStyle("PropertyWindow.NormalFont"))
 									.Text(this, &SSettingsEditorCheckoutNotice::HandleUnlockedStatusText)
-									.ColorAndOpacity(FLinearColor::White)
-									.ShadowColorAndOpacity(FLinearColor::Black)
-									.ShadowOffset(FVector2D::UnitVector)
 							]
 					]
 			]
@@ -196,11 +198,6 @@ FText SSettingsEditorCheckoutNotice::HandleUnlockedStatusText() const
 		LOCTEXT("DefaultSettingsNotice_Writable", "These settings are saved in {0}, which is currently writable."), ConfigFilename);
 }
 
-FSlateColor SSettingsEditorCheckoutNotice::GetLockedOrUnlockedStatusBarColor() const
-{
-	return IsUnlocked() ? FLinearColor::Green : FLinearColor::Yellow;
-}
-
 
 EVisibility SSettingsEditorCheckoutNotice::HandleThrobberVisibility() const
 {
diff --git a/Engine/Source/Developer/SharedSettingsWidgets/Public/SSettingsEditorCheckoutNotice.h b/Engine/Source/Developer/SharedSettingsWidgets/Public/SSettingsEditorCheckoutNotice.h
--- a/Engine/Source/Developer/SharedSettingsWidgets/Public/SSettingsEditorCheckoutNotice.h
+++ b/Engine/Source/Developer/SharedSettingsWidgets/Public/SSettingsEditorCheckoutNotice.h
@@ -66,9 +66,6 @@ private:
 	// Callback for getting the status text when the config is unlocked
 	FText HandleUnlockedStatusText() const;
 
-	// Callback for getting the status color, which indicates if it is writeable or not
-	FSlateColor GetLockedOrUnlockedStatusBarColor() const;
-
 	// Callback for getting the visibility of the source control throbber
 	EVisibility HandleThrobberVisibility() const;
 
