Updated checkout warning styles

Checkout warnings now have a yellow-orange background to highlight the error, and checked-out/writable notices now have no color to reduce noise.

diff --git a/Engine/Source/Developer/SharedSettingsWidgets/Private/SSettingsEditorCheckoutNotice.cpp b/Engine/Source/Developer/SharedSettingsWidgets/Private/SSettingsEditorCheckoutNotice.cpp
--- a/Engine/Source/Developer/SharedSettingsWidgets/Private/SSettingsEditorCheckoutNotice.cpp
+++ b/Engine/Source/Developer/SharedSettingsWidgets/Private/SSettingsEditorCheckoutNotice.cpp
@@ -26,7 +26,8 @@ void SSettingsEditorCheckoutNotice::Construct( const FArguments& InArgs )
 	ChildSlot
 	[
 		SNew(SBorder)
-			.BorderImage(FEditorStyle::GetBrush("ToolPanel.GroupBorder"))
+			.BorderImage(FEditorStyle::GetBrush("SettingsEditor.CheckoutWarningBorder"))
+			.BorderBackgroundColor(this, &SSettingsEditorCheckoutNotice::HandleBorderBackgroundColor)
 			[
 				SNew(SWidgetSwitcher)
 					.WidgetIndex(this, &SSettingsEditorCheckoutNotice::HandleNoticeSwitcherWidgetIndex)
@@ -44,7 +45,6 @@ void SSettingsEditorCheckoutNotice::Construct( const FArguments& InArgs )
 							[
 								SNew(SImage)
 								.Image(FEditorStyle::GetBrush("GenericLock"))
-								.ColorAndOpacity(FColor(255, 168, 0))
 							]
 
 						// Locked notice
@@ -56,6 +56,8 @@ void SSettingsEditorCheckoutNotice::Construct( const FArguments& InArgs )
 								SNew(STextBlock)
 									.Text(this, &SSettingsEditorCheckoutNotice::HandleLockedStatusText)
 									.Font(FEditorStyle::GetFontStyle("PropertyWindow.NormalFont"))
+									.ShadowColorAndOpacity(FLinearColor::Black.CopyWithNewOpacity(0.3f))
+									.ShadowOffset(FVector2D::UnitVector)
 							]
 
 						// Check out button
@@ -95,7 +97,6 @@ void SSettingsEditorCheckoutNotice::Construct( const FArguments& InArgs )
 							[
 								SNew(SImage)
 								.Image(FEditorStyle::GetBrush("GenericUnlock"))
-								.ColorAndOpacity(FLinearColor::Green)
 							]
 
 						// Unlocked notice
@@ -208,6 +209,11 @@ EVisibility SSettingsEditorCheckoutNotice::HandleThrobberVisibility() const
 	return EVisibility::Collapsed;
 }
 
+FSlateColor SSettingsEditorCheckoutNotice::HandleBorderBackgroundColor() const
+{
+	static FColor Orange(166, 137, 0);
+	return IsUnlocked() || DefaultConfigQueryInProgress ? FColor(60, 60, 60) : Orange;
+}
 
 void SSettingsEditorCheckoutNotice::Tick(const FGeometry& AllottedGeometry, const double InCurrentTime, const float InDeltaTime)
 {
diff --git a/Engine/Source/Developer/SharedSettingsWidgets/Public/SSettingsEditorCheckoutNotice.h b/Engine/Source/Developer/SharedSettingsWidgets/Public/SSettingsEditorCheckoutNotice.h
--- a/Engine/Source/Developer/SharedSettingsWidgets/Public/SSettingsEditorCheckoutNotice.h
+++ b/Engine/Source/Developer/SharedSettingsWidgets/Public/SSettingsEditorCheckoutNotice.h
@@ -69,6 +69,8 @@ private:
 	// Callback for getting the visibility of the source control throbber
 	EVisibility HandleThrobberVisibility() const;
 
+	// Callback to get the background color of the checkout wanring
+	FSlateColor HandleBorderBackgroundColor() const;
 private:
 
 	// Called when the file was modified externally
diff --git a/Engine/Source/Editor/EditorStyle/Private/SlateEditorStyle.cpp b/Engine/Source/Editor/EditorStyle/Private/SlateEditorStyle.cpp
--- a/Engine/Source/Editor/EditorStyle/Private/SlateEditorStyle.cpp
+++ b/Engine/Source/Editor/EditorStyle/Private/SlateEditorStyle.cpp
@@ -5747,6 +5747,8 @@ void FSlateEditorStyle::FStyle::SetupAutomationStyles()
 		// Settings editor
 		Set("SettingsEditor.GoodIcon", new IMAGE_BRUSH("Settings/Settings_Good", Icon40x40));
 		Set("SettingsEditor.WarningIcon", new IMAGE_BRUSH("Settings/Settings_Warning", Icon40x40));
+
+		Set( "SettingsEditor.CheckoutWarningBorder", new BOX_BRUSH( "Common/GroupBorderLight", FMargin(4.0f/16.0f) ) );
 	}
 
 	{
