CIS fix for mac

Clang does not allow multiple implicit type conversions. Also added a "checking file state..." message to the notice.

diff --git a/Engine/Source/Developer/SharedSettingsWidgets/Private/SSettingsEditorCheckoutNotice.cpp b/Engine/Source/Developer/SharedSettingsWidgets/Private/SSettingsEditorCheckoutNotice.cpp
--- a/Engine/Source/Developer/SharedSettingsWidgets/Private/SSettingsEditorCheckoutNotice.cpp
+++ b/Engine/Source/Developer/SharedSettingsWidgets/Private/SSettingsEditorCheckoutNotice.cpp
@@ -185,9 +185,14 @@ FText SSettingsEditorCheckoutNotice::HandleLockedStatusText() const
 {
 	FText ConfigFilename = FText::FromString(FPaths::GetCleanFilename(ConfigFilePath.Get()));
 
+	if (DefaultConfigQueryInProgress)
+	{
+		return FText::Format(LOCTEXT("DefaultSettingsNotice_Source", "These settings are saved in {0}. Checking file state..."), ConfigFilename);
+	}
+	
 	return FText::Format(ISourceControlModule::Get().IsEnabled() ?
 		LOCTEXT("DefaultSettingsNotice_WithSourceControl", "These settings are saved in {0}, which is not currently checked out.") :
-		LOCTEXT("DefaultSettingsNotice_Source", "These settings are saved {0}, which is not currently writable."), ConfigFilename);
+		LOCTEXT("DefaultSettingsNotice_Source", "These settings are saved in {0}, which is not currently writable."), ConfigFilename);
 }
 
 FText SSettingsEditorCheckoutNotice::HandleUnlockedStatusText() const
@@ -212,7 +217,9 @@ EVisibility SSettingsEditorCheckoutNotice::HandleThrobberVisibility() const
 FSlateColor SSettingsEditorCheckoutNotice::HandleBorderBackgroundColor() const
 {
 	static FColor Orange(166, 137, 0);
-	return IsUnlocked() || DefaultConfigQueryInProgress ? FColor(60, 60, 60) : Orange;
+
+	const FLinearColor FinalColor = (IsUnlocked() || DefaultConfigQueryInProgress) ? FColor(60, 60, 60) : Orange;
+	return FinalColor;
 }
 
 void SSettingsEditorCheckoutNotice::Tick(const FGeometry& AllottedGeometry, const double InCurrentTime, const float InDeltaTime)
