Settings: Slightly rephrased checkout notice message

diff --git a/Engine/Source/Developer/SharedSettingsWidgets/Private/SSettingsEditorCheckoutNotice.cpp b/Engine/Source/Developer/SharedSettingsWidgets/Private/SSettingsEditorCheckoutNotice.cpp
--- a/Engine/Source/Developer/SharedSettingsWidgets/Private/SSettingsEditorCheckoutNotice.cpp
+++ b/Engine/Source/Developer/SharedSettingsWidgets/Private/SSettingsEditorCheckoutNotice.cpp
@@ -195,8 +195,8 @@ FText SSettingsEditorCheckoutNotice::HandleLockedStatusText() const
 	}
 	
 	return FText::Format(ISourceControlModule::Get().IsEnabled() ?
-		LOCTEXT("DefaultSettingsNotice_WithSourceControl", "These settings are saved in {0}, which is not currently checked out.") :
-		LOCTEXT("DefaultSettingsNotice_Source", "These settings are saved in {0}, which is not currently writable."), ConfigFilename);
+		LOCTEXT("DefaultSettingsNotice_WithSourceControl", "These settings are saved in {0}, which is currently NOT checked out.") :
+		LOCTEXT("DefaultSettingsNotice_Source", "These settings are saved in {0}, which is currently NOT writable."), ConfigFilename);
 }
 
 FText SSettingsEditorCheckoutNotice::HandleUnlockedStatusText() const
