Added make file writable button, which allows file to be edited when non-writable and not connected to source control.

UE-6393

diff --git a/Engine/Source/Developer/SharedSettingsWidgets/Private/SSettingsEditorCheckoutNotice.cpp b/Engine/Source/Developer/SharedSettingsWidgets/Private/SSettingsEditorCheckoutNotice.cpp
--- a/Engine/Source/Developer/SharedSettingsWidgets/Private/SSettingsEditorCheckoutNotice.cpp
+++ b/Engine/Source/Developer/SharedSettingsWidgets/Private/SSettingsEditorCheckoutNotice.cpp
@@ -175,11 +175,13 @@ FText SSettingsEditorCheckoutNotice::HandleCheckOutButtonToolTip( ) const
 	return LOCTEXT("MakeWritableTooltip", "Make the default configuration file that holds these settings writable.");
 }
 
-EVisibility SSettingsEditorCheckoutNotice::HandleCheckOutButtonVisibility( ) const
+EVisibility SSettingsEditorCheckoutNotice::HandleCheckOutButtonVisibility() const
 {
-	if (ISourceControlModule::Get().IsEnabled() && ISourceControlModule::Get().GetProvider().IsAvailable())
+	// Display for checking out the file, or for making writable
+	if (ISourceControlModule::Get().IsEnabled() && ISourceControlModule::Get().GetProvider().IsAvailable() && !DefaultConfigQueryInProgress || 
+		!ISourceControlModule::Get().IsEnabled() && (DefaultConfigQueryInProgress || DefaultConfigCheckOutNeeded))
 	{
-		return !DefaultConfigQueryInProgress ? EVisibility::Visible : EVisibility::Hidden;
+		return EVisibility::Visible;
 	}
 
 	return EVisibility::Hidden;
