Fix CIS.

diff --git a/Engine/Source/Developer/SharedSettingsWidgets/Private/SSettingsEditorCheckoutNotice.cpp b/Engine/Source/Developer/SharedSettingsWidgets/Private/SSettingsEditorCheckoutNotice.cpp
--- a/Engine/Source/Developer/SharedSettingsWidgets/Private/SSettingsEditorCheckoutNotice.cpp
+++ b/Engine/Source/Developer/SharedSettingsWidgets/Private/SSettingsEditorCheckoutNotice.cpp
@@ -178,7 +178,7 @@ FText SSettingsEditorCheckoutNotice::HandleCheckOutButtonToolTip( ) const
 EVisibility SSettingsEditorCheckoutNotice::HandleCheckOutButtonVisibility() const
 {
 	// Display for checking out the file, or for making writable
-	if (ISourceControlModule::Get().IsEnabled() && ISourceControlModule::Get().GetProvider().IsAvailable() && !DefaultConfigQueryInProgress || 
+	if ((ISourceControlModule::Get().IsEnabled() && ISourceControlModule::Get().GetProvider().IsAvailable() && !DefaultConfigQueryInProgress) || 
 		!ISourceControlModule::Get().IsEnabled() && (DefaultConfigQueryInProgress || DefaultConfigCheckOutNeeded))
 	{
 		return EVisibility::Visible;
