Fix CIS (this time for reals).

diff --git a/Engine/Source/Developer/SharedSettingsWidgets/Private/SSettingsEditorCheckoutNotice.cpp b/Engine/Source/Developer/SharedSettingsWidgets/Private/SSettingsEditorCheckoutNotice.cpp
--- a/Engine/Source/Developer/SharedSettingsWidgets/Private/SSettingsEditorCheckoutNotice.cpp
+++ b/Engine/Source/Developer/SharedSettingsWidgets/Private/SSettingsEditorCheckoutNotice.cpp
@@ -179,7 +179,7 @@ EVisibility SSettingsEditorCheckoutNotice::HandleCheckOutButtonVisibility() cons
 {
 	// Display for checking out the file, or for making writable
 	if ((ISourceControlModule::Get().IsEnabled() && ISourceControlModule::Get().GetProvider().IsAvailable() && !DefaultConfigQueryInProgress) || 
-		!ISourceControlModule::Get().IsEnabled() && (DefaultConfigQueryInProgress || DefaultConfigCheckOutNeeded))
+		(!ISourceControlModule::Get().IsEnabled() && (DefaultConfigQueryInProgress || DefaultConfigCheckOutNeeded)))
 	{
 		return EVisibility::Visible;
 	}
