Make 10.13.6 the minimum version of macOS required by UE4 for games and the editor

diff --git a/Engine/Source/Developer/DesktopPlatform/Private/Mac/DesktopPlatformMac.cpp b/Engine/Source/Developer/DesktopPlatform/Private/Mac/DesktopPlatformMac.cpp
--- a/Engine/Source/Developer/DesktopPlatform/Private/Mac/DesktopPlatformMac.cpp
+++ b/Engine/Source/Developer/DesktopPlatform/Private/Mac/DesktopPlatformMac.cpp
@@ -281,7 +281,7 @@ bool FDesktopPlatformMac::OpenDirectoryDialog(const void* ParentWindowHandle, co
 
 			NSInteger Result = [Panel runModal];
 
-			if (Result == NSFileHandlingPanelOKButton)
+			if (Result == NSModalResponseOK)
 			{
 				NSURL *FolderURL = [[Panel URLs] objectAtIndex: 0];
 				TCHAR FolderPath[MAC_MAX_PATH];
@@ -425,7 +425,7 @@ bool FDesktopPlatformMac::FileDialogShared(bool bSave, const void* ParentWindowH
 			NSInteger Result = [Panel runModal];
 			[AccessoryView release];
 
-			if (Result == NSFileHandlingPanelOKButton)
+			if (Result == NSModalResponseOK)
 			{
 				if (bSave)
 				{
diff --git a/Engine/Source/Developer/SharedSettingsWidgets/Private/ShaderFormatsPropertyDetails.cpp b/Engine/Source/Developer/SharedSettingsWidgets/Private/ShaderFormatsPropertyDetails.cpp
--- a/Engine/Source/Developer/SharedSettingsWidgets/Private/ShaderFormatsPropertyDetails.cpp
+++ b/Engine/Source/Developer/SharedSettingsWidgets/Private/ShaderFormatsPropertyDetails.cpp
@@ -75,10 +75,10 @@ FText FShaderFormatsPropertyDetails::GetFriendlyNameFromRHINameMac(const FString
 			FriendlyRHIName = LOCTEXT("MetalMRTTV", "tvOS Metal Desktop Renderer (SM5, Metal 1.2+, tvOS 10.0 or later)");
 			break;
 		case SP_METAL_SM5_NOTESS:
-			FriendlyRHIName = LOCTEXT("MetalSM5_NOTESS", "Mac Metal Desktop Renderer without Tessellation (SM5, Metal 1.2+, macOS Sierra 10.12.6 or later)");
+			FriendlyRHIName = LOCTEXT("MetalSM5_NOTESS", "Mac Metal Desktop Renderer without Tessellation (SM5, Metal 2.0+, macOS High Sierra 10.13.6 or later)");
 			break;
 		case SP_METAL_SM5:
-			FriendlyRHIName = LOCTEXT("MetalSM5", "Mac Metal Desktop Renderer with Tessellation (SM5, Metal 1.2+, macOS Sierra 10.12.6 or later)");
+			FriendlyRHIName = LOCTEXT("MetalSM5", "Mac Metal Desktop Renderer with Tessellation (SM5, Metal 2.0+, macOS High Sierra 10.13.6 or later)");
 			break;
 		case SP_METAL_MACES3_1:
 			FriendlyRHIName = LOCTEXT("MetalES3.1", "Mac Metal High-End Mobile Preview (ES3.1)");
diff --git a/Engine/Source/Programs/UnrealBuildTool/Platform/Mac/MacToolChain.cs b/Engine/Source/Programs/UnrealBuildTool/Platform/Mac/MacToolChain.cs
--- a/Engine/Source/Programs/UnrealBuildTool/Platform/Mac/MacToolChain.cs
+++ b/Engine/Source/Programs/UnrealBuildTool/Platform/Mac/MacToolChain.cs
@@ -32,12 +32,12 @@ namespace UnrealBuildTool
 		/// <summary>
 		/// Which version of the Mac OS X to allow at run time
 		/// </summary>
-		public string MacOSVersion = "10.12";
+		public string MacOSVersion = "10.13";
 
 		/// <summary>
 		/// Minimum version of Mac OS X to actually run on, running on earlier versions will display the system minimum version error dialog and exit.
 		/// </summary>
-		public string MinMacOSVersion = "10.12.6";
+		public string MinMacOSVersion = "10.13.6";
 
 		/// <summary>
 		/// Directory for the developer binaries
diff --git a/Engine/Source/Runtime/AVIWriter/Private/AVIWriter.cpp b/Engine/Source/Runtime/AVIWriter/Private/AVIWriter.cpp
--- a/Engine/Source/Runtime/AVIWriter/Private/AVIWriter.cpp
+++ b/Engine/Source/Runtime/AVIWriter/Private/AVIWriter.cpp
@@ -476,7 +476,7 @@ public:
 			if (Options.CompressionQuality.IsSet())
 			{
 				VideoSettings = [NSDictionary dictionaryWithObjectsAndKeys:
-								 AVVideoCodecJPEG, AVVideoCodecKey,
+								 AVVideoCodecTypeJPEG, AVVideoCodecKey,
 								 [NSDictionary dictionaryWithObjectsAndKeys:[NSNumber numberWithFloat:Options.CompressionQuality.GetValue()], AVVideoQualityKey, nil], AVVideoCompressionPropertiesKey,
 								 [NSNumber numberWithInt:Options.Width], AVVideoWidthKey,
 								 [NSNumber numberWithInt:Options.Height], AVVideoHeightKey,
@@ -485,7 +485,7 @@ public:
 			else
 			{
 				VideoSettings = [NSDictionary dictionaryWithObjectsAndKeys:
-								 AVVideoCodecH264, AVVideoCodecKey,
+								 AVVideoCodecTypeH264, AVVideoCodecKey,
 								 [NSNumber numberWithInt:Options.Width], AVVideoWidthKey,
 								 [NSNumber numberWithInt:Options.Height], AVVideoHeightKey,
 								 nil];
