Editor: fix crash adding RHIs in Windows settings (UE-72863).

- Crash happened if TargetedRHIs contained formats without friendly names (e.g. SF_VULKAN_ES31 in ActionRPG).
- Change by [REDACTED_PERSON]

diff --git a/Engine/Source/Developer/SharedSettingsWidgets/Private/ShaderFormatsPropertyDetails.cpp b/Engine/Source/Developer/SharedSettingsWidgets/Private/ShaderFormatsPropertyDetails.cpp
--- a/Engine/Source/Developer/SharedSettingsWidgets/Private/ShaderFormatsPropertyDetails.cpp
+++ b/Engine/Source/Developer/SharedSettingsWidgets/Private/ShaderFormatsPropertyDetails.cpp
@@ -184,7 +184,8 @@ void FShaderFormatsPropertyDetails::OnTargetedRHIChanged(ECheckBoxState InNewVal
 				int32 InsertIndex = 0;
 				for (; InsertIndex < Array.Num(); ++InsertIndex)
 				{
-					if (InIndex < ShaderFormatOrder[*Array[InsertIndex]])
+					const int32* ShaderFormatIndex = ShaderFormatOrder.Find(*Array[InsertIndex]);
+					if (ShaderFormatIndex != nullptr && InIndex < *ShaderFormatIndex) 
 					{
 						break;
 					}
