Fixed null pointer access in setting helpers

diff --git a/Engine/Source/Developer/SharedSettingsWidgets/Private/SSettingsEditorCheckoutNotice.cpp b/Engine/Source/Developer/SharedSettingsWidgets/Private/SSettingsEditorCheckoutNotice.cpp
--- a/Engine/Source/Developer/SharedSettingsWidgets/Private/SSettingsEditorCheckoutNotice.cpp
+++ b/Engine/Source/Developer/SharedSettingsWidgets/Private/SSettingsEditorCheckoutNotice.cpp
@@ -29,11 +29,7 @@ namespace SettingsHelpers
 		{
 			ISourceControlProvider& SourceControlProvider = ISourceControlModule::Get().GetProvider();
 			FSourceControlStatePtr SourceControlState = SourceControlProvider.GetState(InFileToCheckOut, bForceSourceControlUpdate ? EStateCacheUsage::ForceUpdate : EStateCacheUsage::Use);
-
-			if (SourceControlState->IsCheckedOut() || SourceControlState->IsAdded())
-			{
-				return true;
-			}
+			return SourceControlState.IsValid() && (SourceControlState->IsCheckedOut() || SourceControlState->IsAdded());
 		}
 
 		return false;
@@ -45,11 +41,7 @@ namespace SettingsHelpers
 		{
 			ISourceControlProvider& SourceControlProvider = ISourceControlModule::Get().GetProvider();
 			FSourceControlStatePtr SourceControlState = SourceControlProvider.GetState(InFileToCheckOut, bForceSourceControlUpdate ? EStateCacheUsage::ForceUpdate : EStateCacheUsage::Use);
-
-			if (SourceControlState->IsSourceControlled())
-			{
-				return true;
-			}
+			return SourceControlState.IsValid() && SourceControlState->IsSourceControlled();
 		}
 
 		return false;
