Updated project settings checkout warning indicator

diff --git a/Engine/Source/Developer/SharedSettingsWidgets/Private/SSettingsEditorCheckoutNotice.cpp b/Engine/Source/Developer/SharedSettingsWidgets/Private/SSettingsEditorCheckoutNotice.cpp
--- a/Engine/Source/Developer/SharedSettingsWidgets/Private/SSettingsEditorCheckoutNotice.cpp
+++ b/Engine/Source/Developer/SharedSettingsWidgets/Private/SSettingsEditorCheckoutNotice.cpp
@@ -175,92 +175,83 @@ void SSettingsEditorCheckoutNotice::Construct( const FArguments& InArgs )
 	// default configuration notice
 	ChildSlot
 	[
-		SNew(SBorder)
-			.BorderImage(FEditorStyle::GetBrush("SettingsEditor.CheckoutWarningBorder"))
-			.BorderBackgroundColor(this, &SSettingsEditorCheckoutNotice::HandleBorderBackgroundColor)
+		SNew(SWidgetSwitcher)
+		.WidgetIndex(this, &SSettingsEditorCheckoutNotice::HandleNoticeSwitcherWidgetIndex)
+
+		// Locked slot
+		+ SWidgetSwitcher::Slot()
+		[
+			SNew(SHorizontalBox)
+
+			// Locked icon
+			+ SHorizontalBox::Slot()
+			.AutoWidth()
+			.Padding(Padding)
+			.VAlign(VAlign_Center)
 			[
-				SNew(SWidgetSwitcher)
-					.WidgetIndex(this, &SSettingsEditorCheckoutNotice::HandleNoticeSwitcherWidgetIndex)
-
-				// Locked slot
-				+ SWidgetSwitcher::Slot()
-					[
-						SNew(SHorizontalBox)
-
-						// Locked icon
-						+ SHorizontalBox::Slot()
-							.AutoWidth()
-							.Padding(Padding)
-							.VAlign(VAlign_Center)
-							[
-								SNew(SImage)
-								.Image(FEditorStyle::GetBrush("GenericLock"))
-							]
-
-						// Locked notice
-						+ SHorizontalBox::Slot()
-							.FillWidth(1.0f)
-							.Padding(FMargin(0.f, Padding, Padding, Padding))
-							.VAlign(VAlign_Center)
-							[
-								SNew(STextBlock)
-									.Text(this, &SSettingsEditorCheckoutNotice::HandleLockedStatusText)
-									.Font(FEditorStyle::GetFontStyle("PropertyWindow.NormalFont"))
-									.ShadowColorAndOpacity(FLinearColor::Black.CopyWithNewOpacity(0.3f))
-									.ShadowOffset(FVector2D::UnitVector)
-							]
-
-						// Check out button
-						+ SHorizontalBox::Slot()
-							.Padding(FMargin(0.f, 0.f, Padding, 0.f))
-							.AutoWidth()
-							.VAlign(VAlign_Center)
-							[
-								SNew(SButton)
-									.OnClicked(this, &SSettingsEditorCheckoutNotice::HandleCheckOutButtonClicked)
-									.Text(this, &SSettingsEditorCheckoutNotice::HandleCheckOutButtonText)
-									.ToolTipText(this, &SSettingsEditorCheckoutNotice::HandleCheckOutButtonToolTip)
-									.Visibility(this, &SSettingsEditorCheckoutNotice::HandleCheckOutButtonVisibility)
-							]
-
-						// Source control status throbber
-						+ SHorizontalBox::Slot()
-							.AutoWidth()
-							.VAlign(VAlign_Center)
-							.Padding(0.0f, 1.0f)
-							[
-								SNew(SThrobber)
-									.Visibility(this, &SSettingsEditorCheckoutNotice::HandleThrobberVisibility)
-							]
-					]
-
-				// Unlocked slot
-				+ SWidgetSwitcher::Slot()
-					[
-						SNew(SHorizontalBox)
-
-						// Unlocked icon
-						+ SHorizontalBox::Slot()
-							.AutoWidth()
-							.Padding(Padding)
-							.VAlign(VAlign_Center)
-							[
-								SNew(SImage)
-								.Image(FEditorStyle::GetBrush("GenericUnlock"))
-							]
-
-						// Unlocked notice
-						+ SHorizontalBox::Slot()
-							.AutoWidth()
-							.Padding(FMargin(0.f, Padding, Padding, Padding))
-							.VAlign(VAlign_Center)
-							[
-								SNew(STextBlock)
-									.Font(FEditorStyle::GetFontStyle("PropertyWindow.NormalFont"))
-									.Text(this, &SSettingsEditorCheckoutNotice::HandleUnlockedStatusText)
-							]
-					]
+				SNew(SImage)
+				.Image(FAppStyle::Get().GetBrush("Icons.Lock"))
 			]
+
+			// Locked notice
+			+ SHorizontalBox::Slot()
+			.FillWidth(1.0f)
+			.Padding(FMargin(0.f, Padding, Padding, Padding))
+			.VAlign(VAlign_Center)
+			[
+				SNew(STextBlock)
+				.Text(this, &SSettingsEditorCheckoutNotice::HandleLockedStatusText)
+				.ColorAndOpacity(this, &SSettingsEditorCheckoutNotice::GetLockedStatusTextColor)
+								
+			]
+
+			// Check out button
+			+ SHorizontalBox::Slot()
+			.Padding(FMargin(0.f, 0.f, Padding, 0.f))
+			.AutoWidth()
+			.VAlign(VAlign_Center)
+			[
+				SNew(SButton)
+				.OnClicked(this, &SSettingsEditorCheckoutNotice::HandleCheckOutButtonClicked)
+				.Text(this, &SSettingsEditorCheckoutNotice::HandleCheckOutButtonText)
+				.ToolTipText(this, &SSettingsEditorCheckoutNotice::HandleCheckOutButtonToolTip)
+				.Visibility(this, &SSettingsEditorCheckoutNotice::HandleCheckOutButtonVisibility)
+			]
+			// Source control status throbber
+			+ SHorizontalBox::Slot()
+			.AutoWidth()
+			.VAlign(VAlign_Center)
+			.Padding(0.0f, 1.0f)
+			[
+				SNew(SThrobber)
+				.Visibility(this, &SSettingsEditorCheckoutNotice::HandleThrobberVisibility)
+			]
+		]
+		// Unlocked slot
+		+ SWidgetSwitcher::Slot()
+		[
+			SNew(SHorizontalBox)
+
+			// Unlocked icon
+			+ SHorizontalBox::Slot()
+			.AutoWidth()
+			.Padding(Padding)
+			.VAlign(VAlign_Center)
+			[
+				SNew(SImage)
+				.Image(FAppStyle::Get().GetBrush("Icons.Unlock"))
+			]
+
+			// Unlocked notice
+			+ SHorizontalBox::Slot()
+			.AutoWidth()
+			.Padding(FMargin(0.f, Padding, Padding, Padding))
+			.VAlign(VAlign_Center)
+			[
+				SNew(STextBlock)
+				.Text(this, &SSettingsEditorCheckoutNotice::HandleUnlockedStatusText)
+			]
+		]
 	];
 }
 
@@ -353,14 +344,15 @@ EVisibility SSettingsEditorCheckoutNotice::HandleThrobberVisibility() const
 	return EVisibility::Collapsed;
 }
 
-FSlateColor SSettingsEditorCheckoutNotice::HandleBorderBackgroundColor() const
+FSlateColor SSettingsEditorCheckoutNotice::GetLockedStatusTextColor() const
 {
-	static FColor Orange(166, 137, 0);
+	static const FName Foreground("Colors.Foreground");
+	static const FName Warning("Colors.AccentYellow");
 
-	const FLinearColor FinalColor = (IsUnlocked() || DefaultConfigQueryInProgress) ? FColor(60, 60, 60) : Orange;
-	return FinalColor;
+	return (IsUnlocked() || DefaultConfigQueryInProgress) ? FAppStyle::Get().GetSlateColor(Foreground) : FAppStyle::Get().GetSlateColor(Warning);
 }
 
+
 void SSettingsEditorCheckoutNotice::Tick( const FGeometry& AllottedGeometry, const double InCurrentTime, const float InDeltaTime )
 {
 	// cache selected settings object's configuration file state (purposely not in an active timer to make sure it appears while the user is active)
diff --git a/Engine/Source/Developer/SharedSettingsWidgets/Public/SSettingsEditorCheckoutNotice.h b/Engine/Source/Developer/SharedSettingsWidgets/Public/SSettingsEditorCheckoutNotice.h
--- a/Engine/Source/Developer/SharedSettingsWidgets/Public/SSettingsEditorCheckoutNotice.h
+++ b/Engine/Source/Developer/SharedSettingsWidgets/Public/SSettingsEditorCheckoutNotice.h
@@ -85,7 +85,7 @@ private:
 	EVisibility HandleThrobberVisibility() const;
 
 	// Callback to get the background color of the checkout wanring
-	FSlateColor HandleBorderBackgroundColor() const;
+	FSlateColor GetLockedStatusTextColor() const;
 private:
 
 	// Called when the file was modified externally
