Fixed padding issues in the Project Settings menu's section headers.

diff --git a/Engine/Source/Developer/SettingsEditor/Private/Widgets/SSettingsSectionHeader.cpp b/Engine/Source/Developer/SettingsEditor/Private/Widgets/SSettingsSectionHeader.cpp
--- a/Engine/Source/Developer/SettingsEditor/Private/Widgets/SSettingsSectionHeader.cpp
+++ b/Engine/Source/Developer/SettingsEditor/Private/Widgets/SSettingsSectionHeader.cpp
@@ -39,7 +39,7 @@ void SSettingsSectionHeader::Construct(const FArguments& InArgs, const UObject*
 
 
 	ChildSlot
-	.Padding(FMargin(0.0f, 8.0f, 0.0f, 5.0f))
+	.Padding(FMargin(0.0f, 8.0f, 16.0f, 5.0f))
 	[
 		SNew(SVerticalBox)
 		+ SVerticalBox::Slot()
diff --git a/Engine/Source/Developer/SharedSettingsWidgets/Private/SSettingsEditorCheckoutNotice.cpp b/Engine/Source/Developer/SharedSettingsWidgets/Private/SSettingsEditorCheckoutNotice.cpp
--- a/Engine/Source/Developer/SharedSettingsWidgets/Private/SSettingsEditorCheckoutNotice.cpp
+++ b/Engine/Source/Developer/SharedSettingsWidgets/Private/SSettingsEditorCheckoutNotice.cpp
@@ -207,7 +207,7 @@ void SSettingsEditorCheckoutNotice::Construct( const FArguments& InArgs )
 
 			// Check out button
 			+ SHorizontalBox::Slot()
-			.Padding(FMargin(0.f, 0.f, Padding, 0.f))
+			.Padding(FMargin(0))
 			.AutoWidth()
 			.VAlign(VAlign_Center)
 			[
