Updated the PlatformSetupMessage to query the status of the file asynchronously.

Before: 239ms
After: 29ms

diff --git a/Engine/Source/Developer/SharedSettingsWidgets/Private/SPlatformSetupMessage.cpp b/Engine/Source/Developer/SharedSettingsWidgets/Private/SPlatformSetupMessage.cpp
--- a/Engine/Source/Developer/SharedSettingsWidgets/Private/SPlatformSetupMessage.cpp
+++ b/Engine/Source/Developer/SharedSettingsWidgets/Private/SPlatformSetupMessage.cpp
@@ -12,6 +12,7 @@
 #include "ISourceControlProvider.h"
 #include "ISourceControlModule.h"
 #include "SourceControlHelpers.h"
+#include "SourceControlOperations.h"
 #include "Styling/AppStyle.h"
 #include "Widgets/Layout/SWidgetSwitcher.h"
 #include "Framework/Notifications/NotificationManager.h"
@@ -93,6 +94,11 @@ void SPlatformSetupMessage::Construct(const FArguments& InArgs, const FString& I
 		LOCTEXT("ReadyToModifyText", "Platform files are writeable"),
 		FText::GetEmpty());
 
+	TSharedRef<SWidget> GettingStatusWidget = MakeRow(
+		"SettingsEditor.WarningIcon",
+		LOCTEXT("GettingStatusText", "Getting status from source control"),
+		FText::GetEmpty());
+
 	ChildSlot
 	[
 		SNew(SBorder)
@@ -120,6 +126,10 @@ void SPlatformSetupMessage::Construct(const FArguments& InArgs, const FString& I
 			[
 				ReadyToModifyWidget
 			]
+			+SWidgetSwitcher::Slot()
+			[
+				GettingStatusWidget
+			]
 		]
 	];
 
@@ -136,7 +146,7 @@ int32 SPlatformSetupMessage::GetSetupStateAsInt() const
 	return (int32)CachedSetupState;
 }
 
-SPlatformSetupMessage::ESetupState SPlatformSetupMessage::GetSetupStateBasedOnFile(bool bForce) const
+SPlatformSetupMessage::ESetupState SPlatformSetupMessage::GetSetupStateBasedOnFile(bool bInitStatus)
 {
 	if (!FPaths::FileExists(TargetFilename))
 	{
@@ -149,19 +159,30 @@ SPlatformSetupMessage::ESetupState SPlatformSetupMessage::GetSetupStateBasedOnFi
 		{
 			ISourceControlProvider& Provider = SCC.GetProvider();
 
-			FSourceControlStatePtr SourceControlState = Provider.GetState(TargetFilename, bForce ? EStateCacheUsage::ForceUpdate : EStateCacheUsage::Use);
-			if (SourceControlState.IsValid())
+			if (bInitStatus)
 			{
-				if (SourceControlState->IsSourceControlled())
+				TSharedRef<FUpdateStatus, ESPMode::ThreadSafe> UpdateStatusOperation = ISourceControlOperation::Create<FUpdateStatus>();
+				if (Provider.Execute(UpdateStatusOperation, EConcurrency::Asynchronous, FSourceControlOperationComplete::CreateSP(this, &SPlatformSetupMessage::OnSourceControlOperationComplete)) == ECommandResult::Succeeded)
 				{
-					if (SourceControlState->CanCheckout())
-					{
-						return NeedsCheckout;
-					}
+					return GettingStatus;
 				}
-				else
+			}
+			else
+			{
+				FSourceControlStatePtr SourceControlState = Provider.GetState(TargetFilename, EStateCacheUsage::Use);
+				if (SourceControlState.IsValid())
 				{
-					//@TODO: Should we instead try to add the file?
+					if (SourceControlState->IsSourceControlled())
+					{
+						if (SourceControlState->CanCheckout())
+						{
+							return NeedsCheckout;
+						}
+					}
+					else
+					{
+						//@TODO: Should we instead try to add the file?
+					}
 				}
 			}
 		}
@@ -172,6 +193,11 @@ SPlatformSetupMessage::ESetupState SPlatformSetupMessage::GetSetupStateBasedOnFi
 	}
 }
 
+void SPlatformSetupMessage::OnSourceControlOperationComplete(const FSourceControlOperationRef& InOperation, ECommandResult::Type InResult)
+{
+	UpdateCache(false); // the SCC cache can be used now
+}
+
 void SPlatformSetupMessage::UpdateCache(bool bForceUpdate)
 {
 	CachedSetupState = GetSetupStateBasedOnFile(bForceUpdate);
@@ -187,6 +213,7 @@ FSlateColor SPlatformSetupMessage::GetBorderColor() const
 		return FLinearColor::Green;
 	case ReadOnlyFiles:
 	case NeedsCheckout:
+	case GettingStatus:
 	default:
 		return FLinearColor::Yellow;
 	}
diff --git a/Engine/Source/Developer/SharedSettingsWidgets/Public/SPlatformSetupMessage.h b/Engine/Source/Developer/SharedSettingsWidgets/Public/SPlatformSetupMessage.h
--- a/Engine/Source/Developer/SharedSettingsWidgets/Public/SPlatformSetupMessage.h
+++ b/Engine/Source/Developer/SharedSettingsWidgets/Public/SPlatformSetupMessage.h
@@ -5,6 +5,8 @@
 #include "CoreMinimal.h"
 #include "Misc/Attribute.h"
 #include "Input/Reply.h"
+#include "ISourceControlProvider.h"
+#include "ISourceControlOperation.h"
 #include "Styling/SlateColor.h"
 #include "Widgets/DeclarativeSyntaxSupport.h"
 #include "Widgets/SWidget.h"
@@ -22,7 +24,8 @@ class SHAREDSETTINGSWIDGETS_API SPlatformSetupMessage : public SCompoundWidget
 		MissingFiles,
 		NeedsCheckout,
 		ReadOnlyFiles,
-		ReadyToModify
+		ReadyToModify,
+		GettingStatus
 	};
 
 	SLATE_BEGIN_ARGS(SPlatformSetupMessage)
@@ -53,8 +56,10 @@ private:
 
 	FReply OnButtonPressed();
 
+	void OnSourceControlOperationComplete(const FSourceControlOperationRef& InOperation, ECommandResult::Type InResult);
+
 	// Returns the setup state of a specified file
-	ESetupState GetSetupStateBasedOnFile(bool bForceUpdate) const;
+	ESetupState GetSetupStateBasedOnFile(bool bInitStatus);
 
 	// Updates the cache CachedSetupState 
 	void UpdateCache(bool bForceUpdate);
