Add restart required to Mac target RHI changes

diff --git a/Engine/Source/Developer/SharedSettingsWidgets/Private/ShaderFormatsPropertyDetails.cpp b/Engine/Source/Developer/SharedSettingsWidgets/Private/ShaderFormatsPropertyDetails.cpp
--- a/Engine/Source/Developer/SharedSettingsWidgets/Private/ShaderFormatsPropertyDetails.cpp
+++ b/Engine/Source/Developer/SharedSettingsWidgets/Private/ShaderFormatsPropertyDetails.cpp
@@ -85,6 +85,8 @@ void FShaderFormatsPropertyDetails::CreateTargetShaderFormatsPropertyView(
 			];
 		}
 	}
+    
+    ShaderFormatsPropertyHandle->GetProperty()->SetMetaData("ConfigRestartRequired", "true");
 }
 
 
